Automation Commandlet cs - Adding missing copyright

diff --git a/Engine/Source/Programs/RunMutableCommandlet/RunMutableCommandlet.cs b/Engine/Source/Programs/RunMutableCommandlet/RunMutableCommandlet.cs
--- a/Engine/Source/Programs/RunMutableCommandlet/RunMutableCommandlet.cs
+++ b/Engine/Source/Programs/RunMutableCommandlet/RunMutableCommandlet.cs
@@ -1,3 +1,5 @@
+// Copyright Epic Games, Inc. All Rights Reserved.
+
 using AutomationTool;
 using Gauntlet;
 
